FieldNotify: Add first implementation of FieldNotificationTrace. Add the trace track to the RewindDebugger. Add trace when the field value changes.

diff --git a/Engine/Plugins/Developer/FieldNotificationTrace/FieldNotificationTrace.uplugin b/Engine/Plugins/Developer/FieldNotificationTrace/FieldNotificationTrace.uplugin
new file mode 100644
--- /dev/null
+++ b/Engine/Plugins/Developer/FieldNotificationTrace/FieldNotificationTrace.uplugin
@@ -0,0 +1,38 @@
+{
+	"FileVersion": 3,
+
+	"FriendlyName": "Field Notification Trace",
+	"Version": 1,
+	"VersionName": "1.0",
+	"CreatedBy": "Epic Games, Inc.",
+	"CreatedByURL": "https://epicgames.com",
+	"Description": "Add support to trace field notification object.",
+	"Category": "UI",
+	"EnabledByDefault": false,
+	"CanContainContent": false,
+	"IsBetaVersion": true,
+	"IsExperimentalVersion": true,
+	"Installed": false,
+
+	"Modules": [
+		{
+			"Name": "FieldNotificationTrace",
+			"Type": "Runtime",
+			"LoadingPhase": "Default"
+
+		},
+		{
+			"Name": "FieldNotificationTraceEditor",
+			"Type": "Editor",
+			"LoadingPhase": "Default"
+		}
+	],
+
+	"Plugins": [
+		{
+			"Name": "GameplayInsights",
+			"Enabled": true
+		}
+	]
+}
+
diff --git a/Engine/Plugins/Developer/FieldNotificationTrace/Source/FieldNotificationTrace/FieldNotificationTrace.Build.cs b/Engine/Plugins/Developer/FieldNotificationTrace/Source/FieldNotificationTrace/FieldNotificationTrace.Build.cs
new file mode 100644
--- /dev/null
+++ b/Engine/Plugins/Developer/FieldNotificationTrace/Source/FieldNotificationTrace/FieldNotificationTrace.Build.cs
@@ -0,0 +1,35 @@
+// Copyright Epic Games, Inc. All Rights Reserved.
+
+using System.IO;
+using System.Collections.Generic;
+using UnrealBuildTool;
+
+public class FieldNotificationTrace : ModuleRules 
+{
+	public FieldNotificationTrace(ReadOnlyTargetRules Target) : base(Target)
+	{
+		PublicDependencyModuleNames.AddRange(
+			new string[]
+			{
+				"Core",
+				"CoreUObject",
+				"Engine",
+				"FieldNotification",
+			});
+
+		PrivateDependencyModuleNames.AddRange(
+			new string[]
+			{
+				"TraceLog",
+			});
+			
+		if (Target.bBuildEditor)
+		{
+			PrivateDependencyModuleNames.AddRange(
+				new string[]
+				{
+					"GameplayInsights",
+				});
+		}
+	}
+}
diff --git a/Engine/Plugins/Developer/FieldNotificationTrace/Source/FieldNotificationTrace/Private/FieldNotificationTraceModule.cpp b/Engine/Plugins/Developer/FieldNotificationTrace/Source/FieldNotificationTrace/Private/FieldNotificationTraceModule.cpp
new file mode 100644
--- /dev/null
+++ b/Engine/Plugins/Developer/FieldNotificationTrace/Source/FieldNotificationTrace/Private/FieldNotificationTraceModule.cpp
@@ -0,0 +1,25 @@
+// Copyright Epic Games, Inc. All Rights Reserved.
+
+#include "FieldNotificationTraceModule.h"
+#include "Modules/ModuleManager.h"
+
+#define LOCTEXT_NAMESPACE "FieldNotificationTrace"
+
+namespace UE::FieldNotification
+{
+
+// void FFieldNotificationTraceModule::StartupModule()
+// {
+
+// }
+
+// void FFieldNotificationTraceModule::ShutdownModule()
+// {
+// }
+
+} //namespace 
+
+
+IMPLEMENT_MODULE(UE::FieldNotification::FTraceModule, FieldNotificationTrace);
+
+#undef LOCTEXT_NAMESPACE
diff --git a/Engine/Plugins/Developer/FieldNotificationTrace/Source/FieldNotificationTrace/Private/FieldNotificationTraceModule.h b/Engine/Plugins/Developer/FieldNotificationTrace/Source/FieldNotificationTrace/Private/FieldNotificationTraceModule.h
new file mode 100644
--- /dev/null
+++ b/Engine/Plugins/Developer/FieldNotificationTrace/Source/FieldNotificationTrace/Private/FieldNotificationTraceModule.h
@@ -0,0 +1,24 @@
+// Copyright Epic Games, Inc. All Rights Reserved.
+
+#pragma once
+
+#include "Modules/ModuleInterface.h"
+
+namespace UE::FieldNotification
+{
+
+/**
+ *
+ */
+class FTraceModule : public IModuleInterface
+{
+public:
+	FTraceModule() = default;
+
+	//~ Begin IModuleInterface interface
+	// virtual void StartupModule() override;
+	// virtual void ShutdownModule() override;
+	//~ End IModuleInterface interface
+};
+
+} //namespace
diff --git a/Engine/Plugins/Developer/FieldNotificationTrace/Source/FieldNotificationTrace/Private/Trace/FieldNotificationTrace.cpp b/Engine/Plugins/Developer/FieldNotificationTrace/Source/FieldNotificationTrace/Private/Trace/FieldNotificationTrace.cpp
new file mode 100644
--- /dev/null
+++ b/Engine/Plugins/Developer/FieldNotificationTrace/Source/FieldNotificationTrace/Private/Trace/FieldNotificationTrace.cpp
@@ -0,0 +1,250 @@
+// Copyright Epic Games, Inc. All Rights Reserved.
+
+#include "Trace/FieldNotificationTrace.h"
+
+#include "Misc/ScopeRWLock.h"
+#include "Modules/ModuleManager.h"
+#include "ObjectTrace.h"
+
+#if UE_FIELDNOTIFICATION_TRACE_ENABLED
+#define UE_FIELDNOTIFICATION_TRACE_FIELDVALUE 0
+
+
+#if UE_FIELDNOTIFICATION_TRACE_FIELDVALUE
+#include "IGameplayInsightsModule.h"
+#endif
+
+UE_TRACE_CHANNEL(FieldNotificationChannel);
+
+namespace UE::FieldNotification::Private
+{
+
+bool bTraceIsRecording = false;
+FAutoConsoleCommand StartTracingCommand(
+	TEXT("FieldNotification.StartTracing"),
+	TEXT("Turn on the recording of debugging data."),
+	FConsoleCommandWithArgsDelegate::CreateLambda([](const TArray<FString>& Args){ FTrace::StartTracing(); })
+);
+
+FAutoConsoleCommand StopTracingCommand(
+	TEXT("FieldNotification.StopTracing"),
+	TEXT("Turn off the recording of debugging data."),
+	FConsoleCommandWithArgsDelegate::CreateLambda([](const TArray<FString>& Args) { FTrace::StopTracing(); })
+);
+
+uint32 TraceFNameId(FName Name);
+uint32 TraceStringId(const FString& Value);
+
+} // namespace
+
+namespace UE::FieldNotification
+{
+
+UE_TRACE_EVENT_BEGIN(FieldNotification, ObjectBegin)
+UE_TRACE_EVENT_FIELD(uint64, Cycle)
+UE_TRACE_EVENT_FIELD(uint64, ObjectId)
+
+UE_TRACE_EVENT_END()
+
+UE_TRACE_EVENT_BEGIN(FieldNotification, ObjectEnd)
+UE_TRACE_EVENT_FIELD(uint64, Cycle)
+UE_TRACE_EVENT_FIELD(uint64, ObjectId)
+UE_TRACE_EVENT_END()
+
+UE_TRACE_EVENT_BEGIN(FieldNotification, FieldValueChanged)
+UE_TRACE_EVENT_FIELD(uint64, Cycle)
+UE_TRACE_EVENT_FIELD(double, RecordingTime)
+UE_TRACE_EVENT_FIELD(uint64, ObjectId)
+UE_TRACE_EVENT_FIELD(uint32, FieldNotifyId)
+UE_TRACE_EVENT_END()
+
+void FTrace::OutputObjectBegin(TScriptInterface<INotifyFieldValueChanged> Interface)
+{
+	bool bChannelEnabled = UE_TRACE_CHANNELEXPR_IS_ENABLED(FieldNotificationChannel);
+	if (!bChannelEnabled || Interface.GetObject() == nullptr)
+	{
+		return;
+	}
+
+	if (Interface.GetObject()->HasAnyFlags(RF_ClassDefaultObject))
+	{
+		return;
+	}
+
+	TRACE_OBJECT_LIFETIME_BEGIN(Interface.GetObject());
+
+	UE_TRACE_LOG(FieldNotification, ObjectBegin, FieldNotificationChannel)
+		<< ObjectBegin.Cycle(FPlatformTime::Cycles64())
+		<< ObjectBegin.ObjectId(FObjectTrace::GetObjectId(Interface.GetObject()));
+
+	// since the world may or may not be set for those object
+	// need to trace the full path of the object.
+	//a viewmodel may construct other vm, so at replay we need to associate them correctly.
+}
+
+void FTrace::OutputObjectEnd(TScriptInterface<INotifyFieldValueChanged> Interface)
+{
+	bool bChannelEnabled = UE_TRACE_CHANNELEXPR_IS_ENABLED(FieldNotificationChannel);
+	if (!bChannelEnabled || Interface.GetObject() == nullptr)
+	{
+		return;
+	}
+
+	if (Interface.GetObject()->HasAnyFlags(RF_ClassDefaultObject))
+	{
+		return;
+	}
+
+	TRACE_OBJECT_LIFETIME_END(Interface.GetObject())
+
+	UE_TRACE_LOG(FieldNotification, ObjectEnd, FieldNotificationChannel)
+		<< ObjectEnd.Cycle(FPlatformTime::Cycles64())
+		<< ObjectEnd.ObjectId(FObjectTrace::GetObjectId(Interface.GetObject()));
+}
+
+void FTrace::OutputUpdateField(UObject* Instance, UE::FieldNotification::FFieldId Id)
+{
+	bool bChannelEnabled = UE_TRACE_CHANNELEXPR_IS_ENABLED(FieldNotificationChannel);
+	if (!bChannelEnabled || Instance == nullptr || !Id.IsValid())
+	{
+		return;
+	}
+
+	if (Instance->HasAnyFlags(RF_ClassDefaultObject))
+	{
+		return;
+	}
+
+	TRACE_OBJECT(Instance)
+
+	const uint64 StartCycle = FPlatformTime::Cycles64();
+	const uint64 ObjectId = FObjectTrace::GetObjectId(Instance);
+	const uint32 NameId = Private::TraceFNameId(Id.GetName());
+
+#if UE_FIELDNOTIFICATION_TRACE_FIELDVALUE
+	if (const FProperty* Property = Instance->GetClass()->FindPropertyByName(Id.GetName()))
+	{
+		IGameplayInsightsModule& GameplayInsightsModule = FModuleManager::GetModuleChecked<IGameplayInsightsModule>("GameplayInsights");
+		GameplayInsightsModule.TraceObjectProperty(Instance, Property);
+	}
+	else
+	{
+		 //Serialize all the properties. In case it's a function change and the function uses properties that are not FieldNotify
+		IGameplayInsightsModule& GameplayInsightsModule = FModuleManager::GetModuleChecked<IGameplayInsightsModule>("GameplayInsights");
+		GameplayInsightsModule.TraceObjectProperties(Instance);
+	}
+#endif
+
+	UE_TRACE_LOG(FieldNotification, FieldValueChanged, FieldNotificationChannel)
+		<< FieldValueChanged.Cycle(StartCycle)
+		<< FieldValueChanged.RecordingTime(FObjectTrace::GetWorldElapsedTime(Instance->GetWorld()))
+		<< FieldValueChanged.ObjectId(ObjectId)
+		<< FieldValueChanged.FieldNotifyId(NameId);
+}
+
+void FTrace::StartTracing()
+{
+	if (Private::bTraceIsRecording)
+	{
+		return;
+	}
+	Private::bTraceIsRecording = false;
+
+	UE::Trace::ToggleChannel(TEXT("FieldNotificationChannel"), true);
+	UE::Trace::ToggleChannel(TEXT("Object"), true);
+#if UE_FIELDNOTIFICATION_TRACE_FIELDVALUE
+	UE::Trace::ToggleChannel(TEXT("ObjectProperties"), true);
+#endif
+}
+
+void FTrace::StopTracing()
+{
+	UE::Trace::ToggleChannel(TEXT("FieldNotificationChannel"), false);
+}
+
+} // namespace
+
+namespace UE::FieldNotification::Private
+{
+
+UE_TRACE_EVENT_BEGIN(FieldNotification, StringId, NoSync | Important)
+UE_TRACE_EVENT_FIELD(uint32, Id)
+UE_TRACE_EVENT_FIELD(UE::Trace::WideString, Value)
+UE_TRACE_EVENT_END()
+
+struct FTraceStringId
+{
+	uint32 CurrentId = 0;
+	FRWLock SetLock;
+	TMap<FName, uint32> NameToId;
+	TMap<FString, uint32> StringToId;
+} StringIdInstance;
+
+uint32 TraceFNameId(FName Name)
+{
+	if (Name.IsNone())
+	{
+		return 0;
+	}
+
+	// read only mutex
+	{
+		FRWScopeLock ScopeLock(StringIdInstance.SetLock, SLT_ReadOnly);
+		if (uint32* FoundId = StringIdInstance.NameToId.Find(Name))
+		{
+			return *FoundId;
+		}
+	}
+
+	// write mutex
+	uint32 NewId = 0;
+	{
+		FRWScopeLock ScopeLock(StringIdInstance.SetLock, SLT_Write);
+		++StringIdInstance.CurrentId;
+		NewId = StringIdInstance.CurrentId;
+		StringIdInstance.NameToId.Add(Name, NewId);
+	}
+
+	FString Value = Name.ToString();
+	UE_TRACE_LOG(FieldNotification, StringId, FieldNotificationChannel, Value.Len() * sizeof(TCHAR))
+		<< StringId.Id(NewId)
+		<< StringId.Value(*Value, Value.Len());
+
+	return NewId;
+}
+
+uint32 TraceStringId(const FString& Value)
+{
+	if (Value.IsEmpty())
+	{
+		return 0;
+	}
+
+	// read only mutex
+	{
+		FRWScopeLock ScopeLock(StringIdInstance.SetLock, SLT_ReadOnly);
+		if (uint32* FoundId = StringIdInstance.StringToId.Find(Value))
+		{
+			return *FoundId;
+		}
+	}
+
+	// write mutex
+	uint32 NewId = 0;
+	{
+		FRWScopeLock ScopeLock(StringIdInstance.SetLock, SLT_Write);
+		++StringIdInstance.CurrentId;
+		NewId = StringIdInstance.CurrentId;
+		StringIdInstance.StringToId.Add(Value, NewId);
+	}
+
+	UE_TRACE_LOG(FieldNotification, StringId, FieldNotificationChannel, Value.Len() * sizeof(TCHAR))
+		<< StringId.Id(NewId)
+		<< StringId.Value(*Value, Value.Len());
+
+	return NewId;
+}
+
+} // namespace
+
+#endif
diff --git a/Engine/Plugins/Developer/FieldNotificationTrace/Source/FieldNotificationTrace/Public/Trace/FieldNotificationTrace.h b/Engine/Plugins/Developer/FieldNotificationTrace/Source/FieldNotificationTrace/Public/Trace/FieldNotificationTrace.h
new file mode 100644
--- /dev/null
+++ b/Engine/Plugins/Developer/FieldNotificationTrace/Source/FieldNotificationTrace/Public/Trace/FieldNotificationTrace.h
@@ -0,0 +1,49 @@
+// Copyright Epic Games, Inc. All Rights Reserved.
+
+#pragma once
+
+#include "CoreMinimal.h"
+#include "ProfilingDebugging/TraceAuxiliary.h"
+#include "Trace/Config.h"
+#include "Trace/Trace.h"
+#include "INotifyFieldValueChanged.h"
+#include "FieldNotificationId.h"
+
+#ifndef UE_FIELDNOTIFICATION_TRACE_ENABLED 
+#define UE_FIELDNOTIFICATION_TRACE_ENABLED (UE_TRACE_ENABLED && !IS_PROGRAM && !UE_BUILD_SHIPPING)
+#endif
+
+#if UE_FIELDNOTIFICATION_TRACE_ENABLED
+
+namespace UE::FieldNotification
+{
+
+class FTrace : public FNoncopyable
+{
+public:
+	FIELDNOTIFICATIONTRACE_API static void OutputObjectBegin(TScriptInterface<INotifyFieldValueChanged> Interface);
+	FIELDNOTIFICATIONTRACE_API static void OutputObjectEnd(TScriptInterface<INotifyFieldValueChanged> Interface);
+	FIELDNOTIFICATIONTRACE_API static void OutputUpdateField(UObject* Instance, FFieldId Id);
+
+	static void StartTracing();
+	static void StopTracing();
+};
+
+} //namespace
+
+#define UE_TRACE_FIELDNOTIFICATION_LIFETIME_BEGIN(Interface) \
+::UE::FieldNotification::FTrace::OutputObjectBegin(Interface);
+
+#define UE_TRACE_FIELDNOTIFICATION_LIFETIME_END(Interface) \
+::UE::FieldNotification::FTrace::OutputObjectEnd(Interface);
+
+#define UE_TRACE_FIELDNOTIFICATION_FIELD_VALUE_CHANGED(Object, InField) \
+::UE::FieldNotification::FTrace::OutputUpdateField(Object, InField);
+
+#else //UE_FIELDNOTIFICATION_TRACE_ENABLED
+
+#define UE_TRACE_FIELDNOTIFICATION_LIFETIME_BEGIN(Interface)
+#define UE_TRACE_FIELDNOTIFICATION_LIFETIME_END(Interface)
+#define UE_TRACE_FIELDNOTIFICATION_FIELD_VALUE_CHANGED(Object, InField)
+
+#endif //UE_FIELDNOTIFICATION_TRACE_ENABLED
diff --git a/Engine/Plugins/Developer/FieldNotificationTrace/Source/FieldNotificationTraceEditor/FieldNotificationTraceEditor.Build.cs b/Engine/Plugins/Developer/FieldNotificationTrace/Source/FieldNotificationTraceEditor/FieldNotificationTraceEditor.Build.cs
new file mode 100644
--- /dev/null
+++ b/Engine/Plugins/Developer/FieldNotificationTrace/Source/FieldNotificationTraceEditor/FieldNotificationTraceEditor.Build.cs
@@ -0,0 +1,33 @@
+// Copyright Epic Games, Inc. All Rights Reserved.
+
+using System.IO;
+using System.Collections.Generic;
+using UnrealBuildTool;
+
+public class FieldNotificationTraceEditor : ModuleRules 
+{
+	public FieldNotificationTraceEditor(ReadOnlyTargetRules Target) : base(Target)
+	{
+		PublicDependencyModuleNames.AddRange(
+			new string[]
+			{
+				"Core",
+				"CoreUObject",
+				"Engine",
+				"FieldNotification",
+				"FieldNotificationTrace",
+			});
+
+		PrivateDependencyModuleNames.AddRange(
+			new string[]
+			{
+				"GameplayInsights",
+				"RewindDebuggerInterface",
+				"Slate",
+				"SlateCore",
+				"TraceAnalysis",
+				"TraceLog",
+				"TraceServices",
+			});
+	}
+}
diff --git a/Engine/Plugins/Developer/FieldNotificationTrace/Source/FieldNotificationTraceEditor/Private/FieldNotificationRewindDebugger.cpp b/Engine/Plugins/Developer/FieldNotificationTrace/Source/FieldNotificationTraceEditor/Private/FieldNotificationRewindDebugger.cpp
new file mode 100644
--- /dev/null
+++ b/Engine/Plugins/Developer/FieldNotificationTrace/Source/FieldNotificationTraceEditor/Private/FieldNotificationRewindDebugger.cpp
@@ -0,0 +1,37 @@
+// Copyright Epic Games, Inc. All Rights Reserved.
+
+#include "FieldNotificationRewindDebugger.h"
+
+#include "FieldNotificationTraceProvider.h"
+//#include "IGameplayProvider.h"
+
+namespace UE::FieldNotification
+{
+
+void FRewindDebugger::Update(float DeltaTime, IRewindDebugger* RewindDebugger)
+{
+	const TraceServices::IAnalysisSession* AnalysisSession = RewindDebugger->GetAnalysisSession();
+	TraceServices::FAnalysisSessionReadScope SessionReadScope(*AnalysisSession);
+	
+	const FTraceProvider* TraceProvider = AnalysisSession->ReadProvider<FTraceProvider>(FTraceProvider::ProviderName);
+	//const IGameplayProvider* GameplayProvider = AnalysisSession->ReadProvider<IGameplayProvider>("GameplayProvider");
+	
+	const TraceServices::IFrameProvider& FrameProvider = TraceServices::ReadFrameProvider(*RewindDebugger->GetAnalysisSession());
+	TraceServices::FFrame Frame;
+	if (FrameProvider.GetFrameFromTime(ETraceFrameType::TraceFrameType_Game, RewindDebugger->CurrentTraceTime(), Frame))
+	{
+		// each tick update the UMG Preview Window with what we are trying to debug
+	}
+}
+
+void FRewindDebugger::RecordingStarted(IRewindDebugger*)
+{
+	UE::Trace::ToggleChannel(TEXT("FieldNotificationChannel"), true);
+}
+
+void FRewindDebugger::RecordingStopped(IRewindDebugger*)
+{
+	UE::Trace::ToggleChannel(TEXT("FieldNotificationChannel"), false);
+}
+
+} // namespace
diff --git a/Engine/Plugins/Developer/FieldNotificationTrace/Source/FieldNotificationTraceEditor/Private/FieldNotificationRewindDebugger.h b/Engine/Plugins/Developer/FieldNotificationTrace/Source/FieldNotificationTraceEditor/Private/FieldNotificationRewindDebugger.h
new file mode 100644
--- /dev/null
+++ b/Engine/Plugins/Developer/FieldNotificationTrace/Source/FieldNotificationTraceEditor/Private/FieldNotificationRewindDebugger.h
@@ -0,0 +1,19 @@
+// Copyright Epic Games, Inc. All Rights Reserved.
+#pragma once
+
+#include "CoreMinimal.h"
+#include "IRewindDebugger.h"
+#include "IRewindDebuggerExtension.h"
+
+namespace UE::FieldNotification
+{
+
+class FRewindDebugger : public IRewindDebuggerExtension
+{
+public:
+	virtual void Update(float DeltaTime, IRewindDebugger* RewindDebugger) override;
+	virtual void RecordingStarted(IRewindDebugger* RewindDebugger) override;
+	virtual void RecordingStopped(IRewindDebugger* RewindDebugger) override;
+};
+
+} //namespace
diff --git a/Engine/Plugins/Developer/FieldNotificationTrace/Source/FieldNotificationTraceEditor/Private/FieldNotificationTraceAnalyzer.cpp b/Engine/Plugins/Developer/FieldNotificationTrace/Source/FieldNotificationTraceEditor/Private/FieldNotificationTraceAnalyzer.cpp
new file mode 100644
--- /dev/null
+++ b/Engine/Plugins/Developer/FieldNotificationTrace/Source/FieldNotificationTraceEditor/Private/FieldNotificationTraceAnalyzer.cpp
@@ -0,0 +1,75 @@
+// Copyright Epic Games, Inc. All Rights Reserved.
+
+#include "FieldNotificationTraceAnalyzer.h"
+
+#include "FieldNotificationTraceProvider.h"
+#include "HAL/LowLevelMemTracker.h"
+#include "Serialization/MemoryReader.h"
+#include "TraceServices/Model/AnalysisSession.h"
+#include "TraceServices/Utils.h"
+
+namespace UE::FieldNotification
+{
+
+FTraceAnalyzer::FTraceAnalyzer(TraceServices::IAnalysisSession& InSession, FTraceProvider& InProvider)
+	: Session(InSession)
+	, Provider(InProvider)
+{
+}
+
+void FTraceAnalyzer::OnAnalysisBegin(const FOnAnalysisContext& Context)
+{
+	auto& Builder = Context.InterfaceBuilder;
+
+	Builder.RouteEvent(RouteId_ObjectBegin, "FieldNotification", "ObjectBegin");
+	Builder.RouteEvent(RouteId_ObjectEnd, "FieldNotification", "ObjectEnd");
+	Builder.RouteEvent(RouteId_FieldValueChanged, "FieldNotification", "FieldValueChanged");
+	Builder.RouteEvent(RouteId_StringId, "FieldNotification", "StringId");
+}
+
+bool FTraceAnalyzer::OnEvent(uint16 RouteId, EStyle Style, const FOnEventContext& Context)
+{
+	LLM_SCOPE_BYNAME(TEXT("Insights/FTraceAnalyzer"));
+
+	TraceServices::FAnalysisSessionEditScope _(Session);
+
+	const auto& EventData = Context.EventData;
+	switch (RouteId)
+	{
+		case RouteId_ObjectBegin:
+		{
+			uint64 Cycle = EventData.GetValue<uint64>("Cycle");
+			uint64 ObjectId = EventData.GetValue<uint64>("ObjectId");
+			Provider.AppendObjectBegin(ObjectId, Context.EventTime.AsSeconds(Cycle));
+			break;
+		}
+		case RouteId_ObjectEnd:
+		{
+			uint64 Cycle = EventData.GetValue<uint64>("Cycle");
+			uint64 ObjectId = EventData.GetValue<uint64>("ObjectId");
+			Provider.AppendObjectEnd(ObjectId, Context.EventTime.AsSeconds(Cycle));
+			break;
+		}
+		case RouteId_FieldValueChanged:
+		{
+			uint64 Cycle = EventData.GetValue<uint64>("Cycle");
+			double RecordingTime = EventData.GetValue<double>("RecordingTime");
+			uint64 Id = EventData.GetValue<uint64>("ObjectId");
+			uint32 FieldNotifyId = EventData.GetValue<uint64>("FieldNotifyId");
+			Provider.AppendFieldValueChanged(Id, Context.EventTime.AsSeconds(Cycle), RecordingTime, FieldNotifyId);
+			break;
+		}
+		case RouteId_StringId:
+		{
+			uint32 Id = EventData.GetValue<uint32>("Id");
+			FStringView Value;
+			EventData.GetString("Value", Value);
+			Provider.AppendFieldNotify(Id, FName(Value));
+			break;
+		}
+	}
+
+	return true;
+}
+
+} //namespace
\ No newline at end of file
diff --git a/Engine/Plugins/Developer/FieldNotificationTrace/Source/FieldNotificationTraceEditor/Private/FieldNotificationTraceAnalyzer.h b/Engine/Plugins/Developer/FieldNotificationTrace/Source/FieldNotificationTraceEditor/Private/FieldNotificationTraceAnalyzer.h
new file mode 100644
--- /dev/null
+++ b/Engine/Plugins/Developer/FieldNotificationTrace/Source/FieldNotificationTraceEditor/Private/FieldNotificationTraceAnalyzer.h
@@ -0,0 +1,36 @@
+// Copyright Epic Games, Inc. All Rights Reserved.
+
+#pragma once
+
+#include "Trace/Analyzer.h"
+
+namespace UE::FieldNotification { class FTraceProvider; }
+namespace TraceServices { class IAnalysisSession; }
+
+namespace UE::FieldNotification
+{
+
+class FTraceAnalyzer : public UE::Trace::IAnalyzer
+{
+public:
+	FTraceAnalyzer(TraceServices::IAnalysisSession& InSession, FTraceProvider& InProvider);
+
+	virtual void OnAnalysisBegin(const FOnAnalysisContext& Context) override;
+	virtual void OnAnalysisEnd() override {}
+	virtual bool OnEvent(uint16 RouteId, EStyle Style, const FOnEventContext& Context) override;
+
+private:
+	enum : uint16
+	{
+		RouteId_ObjectBegin,
+		RouteId_ObjectEnd,
+		RouteId_FieldValueChanged,
+		//RouteId_PropertyValueChanged,
+		RouteId_StringId,
+	};
+
+	TraceServices::IAnalysisSession& Session;
+	FTraceProvider& Provider;
+};
+
+} // namespace
diff --git a/Engine/Plugins/Developer/FieldNotificationTrace/Source/FieldNotificationTraceEditor/Private/FieldNotificationTraceEditorModule.cpp b/Engine/Plugins/Developer/FieldNotificationTrace/Source/FieldNotificationTraceEditor/Private/FieldNotificationTraceEditorModule.cpp
new file mode 100644
--- /dev/null
+++ b/Engine/Plugins/Developer/FieldNotificationTrace/Source/FieldNotificationTraceEditor/Private/FieldNotificationTraceEditorModule.cpp
@@ -0,0 +1,33 @@
+// Copyright Epic Games, Inc. All Rights Reserved.
+
+#include "FieldNotificationTraceEditorModule.h"
+#include "Modules/ModuleManager.h"
+
+#include "Features/IModularFeatures.h"
+#include "IRewindDebuggerExtension.h"
+
+#define LOCTEXT_NAMESPACE "FieldNotificationTraceEditor"
+
+namespace UE::FieldNotification
+{
+
+void FTraceEditorModule::StartupModule()
+{
+	IModularFeatures::Get().RegisterModularFeature(IRewindDebuggerExtension::ModularFeatureName, &RewindDebugger);
+	IModularFeatures::Get().RegisterModularFeature(RewindDebugger::IRewindDebuggerTrackCreator::ModularFeatureName, &TrackCreator);
+	IModularFeatures::Get().RegisterModularFeature(TraceServices::ModuleFeatureName, &TraceModule);
+}
+
+
+void FTraceEditorModule::ShutdownModule()
+{
+	IModularFeatures::Get().UnregisterModularFeature(TraceServices::ModuleFeatureName, &TraceModule);
+	IModularFeatures::Get().UnregisterModularFeature(RewindDebugger::IRewindDebuggerTrackCreator::ModularFeatureName, &TrackCreator);
+	IModularFeatures::Get().UnregisterModularFeature(IRewindDebuggerExtension::ModularFeatureName, &RewindDebugger);
+}
+
+} //namespace 
+
+IMPLEMENT_MODULE(UE::FieldNotification::FTraceEditorModule, FieldNotificationTraceEditor);
+
+#undef LOCTEXT_NAMESPACE
diff --git a/Engine/Plugins/Developer/FieldNotificationTrace/Source/FieldNotificationTraceEditor/Private/FieldNotificationTraceEditorModule.h b/Engine/Plugins/Developer/FieldNotificationTrace/Source/FieldNotificationTraceEditor/Private/FieldNotificationTraceEditorModule.h
new file mode 100644
--- /dev/null
+++ b/Engine/Plugins/Developer/FieldNotificationTrace/Source/FieldNotificationTraceEditor/Private/FieldNotificationTraceEditorModule.h
@@ -0,0 +1,33 @@
+// Copyright Epic Games, Inc. All Rights Reserved.
+
+#pragma once
+
+#include "Modules/ModuleInterface.h"
+
+#include "FieldNotificationRewindDebugger.h"
+#include "FieldNotificationTrack.h"
+#include "FieldNotificationTraceServices.h"
+
+namespace UE::FieldNotification
+{
+
+/**
+ *
+ */
+class FTraceEditorModule : public IModuleInterface
+{
+public:
+	FTraceEditorModule() = default;
+
+	//~ Begin IModuleInterface interface
+	virtual void StartupModule() override;
+	 virtual void ShutdownModule() override;
+	//~ End IModuleInterface interface
+
+private:
+	FRewindDebugger RewindDebugger;
+	FTracksCreator TrackCreator;
+	FTraceServiceModule TraceModule;
+};
+
+} //namespace
diff --git a/Engine/Plugins/Developer/FieldNotificationTrace/Source/FieldNotificationTraceEditor/Private/FieldNotificationTraceProvider.cpp b/Engine/Plugins/Developer/FieldNotificationTrace/Source/FieldNotificationTraceEditor/Private/FieldNotificationTraceProvider.cpp
new file mode 100644
--- /dev/null
+++ b/Engine/Plugins/Developer/FieldNotificationTrace/Source/FieldNotificationTraceEditor/Private/FieldNotificationTraceProvider.cpp
@@ -0,0 +1,139 @@
+// Copyright Epic Games, Inc. All Rights Reserved.
+
+#include "FieldNotificationTraceProvider.h"
+#include "ObjectTrace.h"
+
+namespace UE::FieldNotification
+{
+
+FName FTraceProvider::ProviderName("FieldNotificationProvider");
+
+#define LOCTEXT_NAMESPACE "FieldNotificationProvider"
+
+
+
+FTraceProvider::FObject::FObject(TraceServices::IAnalysisSession& InSession)
+	: FieldNotifies(InSession.GetLinearAllocator())
+	, FieldNotifiesRecording(InSession.GetLinearAllocator())
+{
+}
+
+FTraceProvider::FTraceProvider(TraceServices::IAnalysisSession& InSession)
+	: Session(InSession)
+	, ObjectLifetimes(InSession.GetLinearAllocator())
+{
+}
+
+void FTraceProvider::EnumerateObjects(double StartTime, double EndTime, TFunctionRef<void(double InStartTime, double InEndTime, uint32 InDepth, const FObject&)> Callback) const
+{
+	ObjectLifetimes.EnumerateEvents(StartTime, EndTime,
+		[this, &Callback](double InStartTime, double InEndTime, uint32 InDepth, const TSharedRef<FObject>& ExistsMessage)
+		{
+			Callback(InStartTime, InEndTime, InDepth, ExistsMessage.Get());
+			return TraceServices::EEventEnumerate::Continue;
+		});
+}
+
+void FTraceProvider::EnumerateFieldNotifies(uint64 InObjectId, double InStartTime, double InEndTime, TFunctionRef<void(double InStartTime, double InEndTime, uint32 InDepth, const FFieldNotifyEvent&)> Callback) const
+{
+	Session.ReadAccessCheck();
+
+	const TSharedRef<FObject>* Found = ObjectIdToObject.Find(InObjectId);
+	if (Found != nullptr)
+	{
+		(*Found)->FieldNotifies.EnumerateEvents(InStartTime, InEndTime,
+			[&Callback](double InNewStartTime, double InNewEndTime, uint32 InDepth, const FFieldNotifyEvent& ExistsMessage)
+			{
+				Callback(InNewStartTime, InNewEndTime, InDepth, ExistsMessage);
+				return TraceServices::EEventEnumerate::Continue;
+			});
+	}
+}
+
+void FTraceProvider::EnumerateRecordingFieldNotifies(uint64 InObjectId, double InStartTime, double InEndTime, TFunctionRef<void(double InStartTime, double InEndTime, uint32 InDepth, const FFieldNotifyEvent&)> Callback) const
+{
+	Session.ReadAccessCheck();
+
+	const TSharedRef<FObject>* Found = ObjectIdToObject.Find(InObjectId);
+	if (Found != nullptr)
+	{
+		(*Found)->FieldNotifiesRecording.EnumerateEvents(InStartTime, InEndTime,
+			[&Callback](double InNewStartTime, double InNewEndTime, uint32 InDepth, const FFieldNotifyEvent& ExistsMessage)
+			{
+				Callback(InNewStartTime, InNewEndTime, InDepth, ExistsMessage);
+				return TraceServices::EEventEnumerate::Continue;
+			});
+	}
+}
+
+FFieldNotificationId FTraceProvider::GetFieldNotificationId(uint32 InFieldNotifyId) const
+{
+	Session.ReadAccessCheck();
+	return FFieldNotificationId(FieldNotifyIdToFieldNotifyName.FindRef(InFieldNotifyId));
+}
+
+void FTraceProvider::AppendObjectBegin(uint64 InObjectId, double InProfileTime)
+{
+	Session.WriteAccessCheck();
+
+	if (!ObjectIdToObject.Contains(InObjectId))
+	{
+		TSharedRef<FObject> NewObject = MakeShared<FObject>(Session);
+		NewObject->SelfObjectId = InObjectId;
+
+		ObjectIdToObject.Add(InObjectId, NewObject);
+		NewObject->TimelineId = ObjectLifetimes.AppendBeginEvent(InProfileTime, NewObject);
+	}
+}
+
+void FTraceProvider::AppendObjectEnd(uint64 InObjectId, double InProfileTime)
+{
+	Session.WriteAccessCheck();
+
+	if (TSharedRef<FObject>* Found = ObjectIdToObject.Find(InObjectId))
+	{
+		ObjectLifetimes.EndEvent((*Found)->TimelineId, InProfileTime);
+	}
+}
+
+void FTraceProvider::AppendFieldValueChanged(uint64 InObjectId, double InProfileTime, double InRecordingTime, uint32 InFieldNotifyId)
+{
+	Session.WriteAccessCheck();
+
+	FObject* CurrentObject = nullptr;
+	if (TSharedRef<FObject>* Found = ObjectIdToObject.Find(InObjectId))
+	{
+		CurrentObject = &(Found->Get());
+	}
+	else
+	{
+		// make new object
+		TSharedRef<FObject> NewObject = MakeShared<FObject>(Session);
+		NewObject->SelfObjectId = InObjectId;
+
+		ObjectIdToObject.Add(InObjectId, NewObject);
+		NewObject->TimelineId = ObjectLifetimes.AppendBeginEvent(InProfileTime, NewObject);
+
+		CurrentObject = &(NewObject.Get());
+	}
+
+	if (CurrentObject)
+	{
+		FFieldNotifyEvent Value;
+		Value.FieldNotifyId = InFieldNotifyId;
+		CurrentObject->FieldNotifies.AppendEvent(InProfileTime, Value);
+		CurrentObject->FieldNotifiesRecording.AppendEvent(InRecordingTime, Value);
+	}
+}
+
+void FTraceProvider::AppendFieldNotify(uint32 InFieldNotifyId, FName InValue)
+{
+	Session.WriteAccessCheck();
+
+	FieldNotifyIdToFieldNotifyName.FindOrAdd(InFieldNotifyId) = InValue;
+}
+
+
+} //namespace
+
+#undef LOCTEXT_NAMESPACE
diff --git a/Engine/Plugins/Developer/FieldNotificationTrace/Source/FieldNotificationTraceEditor/Private/FieldNotificationTraceProvider.h b/Engine/Plugins/Developer/FieldNotificationTrace/Source/FieldNotificationTraceEditor/Private/FieldNotificationTraceProvider.h
new file mode 100644
--- /dev/null
+++ b/Engine/Plugins/Developer/FieldNotificationTrace/Source/FieldNotificationTraceEditor/Private/FieldNotificationTraceProvider.h
@@ -0,0 +1,69 @@
+// Copyright Epic Games, Inc. All Rights Reserved.
+
+#pragma once
+
+#include "Containers/Map.h"
+#include "Engine/EngineBaseTypes.h"
+#include "FieldNotificationId.h"
+#include "Templates/SharedPointer.h"
+#include "Model/IntervalTimeline.h"
+#include "Model/PointTimeline.h"
+#include "TraceServices/Model/AnalysisSession.h"
+
+namespace TraceServices { class IAnalysisSession; }
+
+namespace UE::FieldNotification
+{
+
+class FTraceProvider : public TraceServices::IProvider
+{
+public:
+	struct FFieldNotifyEvent
+	{
+		uint32 FieldNotifyId = 0;
+	};
+
+	struct FObject
+	{
+		FObject(TraceServices::IAnalysisSession& InSession);
+		uint64 SelfObjectId = 0;
+		uint64 TimelineId = 0;
+		TraceServices::TPointTimeline<FFieldNotifyEvent> FieldNotifies;
+		TraceServices::TPointTimeline<FFieldNotifyEvent> FieldNotifiesRecording;
+	};
+	
+	static FName ProviderName;
+
+	FTraceProvider(TraceServices::IAnalysisSession& InSession);
+
+	void EnumerateObjects(double InStartTime, double InEndTime, TFunctionRef<void(double InStartTime, double InEndTime, uint32 InDepth, const FObject&)> Callback) const;
+	void EnumerateFieldNotifies(uint64 InObjectId, double InStartTime, double InEndTime, TFunctionRef<void(double InStartTime, double InEndTime, uint32 InDepth, const FFieldNotifyEvent&)> Callback) const;
+	void EnumerateRecordingFieldNotifies(uint64 InObjectId, double InStartTime, double InEndTime, TFunctionRef<void(double InStartTime, double InEndTime, uint32 InDepth, const FFieldNotifyEvent&)> Callback) const;
+	FFieldNotificationId GetFieldNotificationId(uint32 InFieldNotifyId) const;
+
+	bool HasData() const
+	{
+		return ObjectIdToObject.Num() > 0;
+	}
+
+	bool HasData(uint64 InObjectId) const
+	{
+		return ObjectIdToObject.Contains(InObjectId);
+	}
+
+	void AppendObjectBegin(uint64 InObjectId, double InProfileTime);
+	void AppendObjectEnd(uint64 InObjectId, double InProfileTime);
+	void AppendFieldValueChanged(uint64 InObjectId, double InProfileTime, double InRecordingTime, uint32 InFieldNotifyId);
+	void AppendFieldNotify(uint32 InFieldNotifyId, FName InValue);
+
+private:
+	TraceServices::IAnalysisSession& Session;
+
+	TMap<uint32, FName> FieldNotifyIdToFieldNotifyName;
+	TMap<uint64, TSharedRef<FObject>> ObjectIdToObject;
+
+	// Timeline containing intervals where an object exists
+	TraceServices::TIntervalTimeline<TSharedRef<FObject>> ObjectLifetimes;
+};
+
+} // namespace
diff --git a/Engine/Plugins/Developer/FieldNotificationTrace/Source/FieldNotificationTraceEditor/Private/FieldNotificationTraceServices.cpp b/Engine/Plugins/Developer/FieldNotificationTrace/Source/FieldNotificationTraceEditor/Private/FieldNotificationTraceServices.cpp
new file mode 100644
--- /dev/null
+++ b/Engine/Plugins/Developer/FieldNotificationTrace/Source/FieldNotificationTraceEditor/Private/FieldNotificationTraceServices.cpp
@@ -0,0 +1,37 @@
+// Copyright Epic Games, Inc. All Rights Reserved.
+
+#include "FieldNotificationTraceServices.h"
+
+#include "FieldNotificationTraceAnalyzer.h"
+#include "FieldNotificationTraceProvider.h"
+
+namespace UE::FieldNotification
+{
+
+FName FTraceServiceModule::ModuleName("FieldNotification");
+
+void FTraceServiceModule::GetModuleInfo(TraceServices::FModuleInfo& OutModuleInfo)
+{
+	OutModuleInfo.Name = ModuleName;
+	OutModuleInfo.DisplayName = TEXT("FieldNotification");
+}
+
+void FTraceServiceModule::OnAnalysisBegin(TraceServices::IAnalysisSession& InSession)
+{
+	TSharedPtr<FTraceProvider> TraceProvider = MakeShared<FTraceProvider>(InSession);
+	InSession.AddProvider(FTraceProvider::ProviderName, TraceProvider);
+
+	InSession.AddAnalyzer(new FTraceAnalyzer(InSession, *TraceProvider));
+}
+
+void FTraceServiceModule::GetLoggers(TArray<const TCHAR *>& OutLoggers)
+{
+	OutLoggers.Add(TEXT("FieldNotification"));
+}
+
+void FTraceServiceModule::GenerateReports(const TraceServices::IAnalysisSession& Session, const TCHAR* CmdLine, const TCHAR* OutputDirectory)
+{
+
+}
+
+} // namespace
diff --git a/Engine/Plugins/Developer/FieldNotificationTrace/Source/FieldNotificationTraceEditor/Private/FieldNotificationTraceServices.h b/Engine/Plugins/Developer/FieldNotificationTrace/Source/FieldNotificationTraceEditor/Private/FieldNotificationTraceServices.h
new file mode 100644
--- /dev/null
+++ b/Engine/Plugins/Developer/FieldNotificationTrace/Source/FieldNotificationTraceEditor/Private/FieldNotificationTraceServices.h
@@ -0,0 +1,31 @@
+// Copyright Epic Games, Inc. All Rights Reserved.
+
+#pragma once
+
+#include "CoreMinimal.h"
+#include "TraceServices/ModuleService.h"
+
+class IAnimationProvider;
+namespace Insights { class ITimingViewSession; }
+
+namespace UE::FieldNotification
+{
+
+class FTraceServiceModule : public TraceServices::IModule
+{
+public:
+	//~ TraceServices::IModule interface
+	virtual void GetModuleInfo(TraceServices::FModuleInfo& OutModuleInfo) override;
+	virtual void OnAnalysisBegin(TraceServices::IAnalysisSession& Session) override;
+	virtual void GetLoggers(TArray<const TCHAR *>& OutLoggers) override;
+	virtual void GenerateReports(const TraceServices::IAnalysisSession& Session, const TCHAR* CmdLine, const TCHAR* OutputDirectory) override;
+	virtual const TCHAR* GetCommandLineArgument() override
+	{
+		return TEXT("fieldnotification");
+	}
+
+private:
+	static FName ModuleName;
+};
+
+}
diff --git a/Engine/Plugins/Developer/FieldNotificationTrace/Source/FieldNotificationTraceEditor/Private/FieldNotificationTrack.cpp b/Engine/Plugins/Developer/FieldNotificationTrace/Source/FieldNotificationTraceEditor/Private/FieldNotificationTrack.cpp
new file mode 100644
--- /dev/null
+++ b/Engine/Plugins/Developer/FieldNotificationTrace/Source/FieldNotificationTraceEditor/Private/FieldNotificationTrack.cpp
@@ -0,0 +1,271 @@
+// Copyright Epic Games, Inc. All Rights Reserved.
+
+#include "FieldNotificationTrack.h"
+
+#include "Editor.h"
+#include "FieldNotificationTraceProvider.h"
+#include "IGameplayProvider.h"
+#include "IRewindDebugger.h"
+
+#include "Subsystems/AssetEditorSubsystem.h"
+
+#define LOCTEXT_NAMESPACE "FieldNotificationTrack"
+
+namespace UE::FieldNotification
+{
+
+static FLinearColor MakeNotifyColor(uint32 InSeed, bool bInLine = false)
+{
+	FRandomStream Stream(InSeed);
+	const uint8 Hue = (uint8)(Stream.FRand() * 255.0f);
+	const uint8 SatVal = bInLine ? 196 : 128;
+	return FLinearColor::MakeFromHSV8(Hue, SatVal, SatVal);
+}
+
+FFieldNotifyTrack::FFieldNotifyTrack(uint64 InObjectId, uint32 InFieldNotifyId, FFieldNotificationId InFieldNotify)
+	: ObjectId(InObjectId)
+	, FieldNotifyId(InFieldNotifyId)
+	, FieldNotify(InFieldNotify)
+{
+	EventData = MakeShared<SEventTimelineView::FTimelineEventData>();
+	Icon = FSlateIcon("EditorStyle", "Sequencer.Tracks.Event", "Sequencer.Tracks.Event");
+}
+
+TSharedPtr<SEventTimelineView::FTimelineEventData> FFieldNotifyTrack::GetEventData() const
+{
+	if (!EventData.IsValid())
+	{
+		EventData = MakeShared<SEventTimelineView::FTimelineEventData>();
+	}
+	
+	EventUpdateRequested++;
+	
+	return EventData;
+}
+
+TSharedPtr<SWidget> FFieldNotifyTrack::GetDetailsViewInternal()
+{
+	return nullptr;
+}
+
+bool FFieldNotifyTrack::UpdateInternal()
+{
+	IRewindDebugger* RewindDebugger = IRewindDebugger::Instance();
+	
+	TRange<double> RecordingTimeRange = RewindDebugger->GetCurrentViewRange();
+	double StartTime = RecordingTimeRange.GetLowerBoundValue();
+	double EndTime = RecordingTimeRange.GetUpperBoundValue();
+	
+	const TraceServices::IAnalysisSession* AnalysisSession = RewindDebugger->GetAnalysisSession();
+	const FTraceProvider* Provider = AnalysisSession->ReadProvider<FTraceProvider>(FTraceProvider::ProviderName);
+	
+	if(Provider)
+	{
+		TRACE_CPUPROFILER_EVENT_SCOPE(FFieldNotifyTrack::UpdateEventPointsInternal);
+		EventUpdateRequested = 0;
+		
+		EventData->Points.SetNum(0,EAllowShrinking::No);
+		EventData->Windows.SetNum(0);
+
+		TraceServices::FAnalysisSessionReadScope SessionReadScope(*AnalysisSession);
+		
+		Provider->EnumerateRecordingFieldNotifies(ObjectId, StartTime, EndTime, [Self=this](double InStartTime, double InEndTime, uint32 InDepth, const FTraceProvider::FFieldNotifyEvent& EvaluationData)
+			{
+				if (EvaluationData.FieldNotifyId == Self->FieldNotifyId)
+				{
+					Self->EventData->Points.Add({ InStartTime, FText(), FText(), FLinearColor::White });
+				}
+			});
+	}
+
+	EventUpdateRequested++;
+
+	bool bChanged = false;
+	return bChanged;
+	
+}
+
+FSlateIcon FFieldNotifyTrack::GetIconInternal()
+{
+	return Icon;
+}
+
+FName FFieldNotifyTrack::GetNameInternal() const
+{
+	return FieldNotify.GetFieldName();
+}
+
+FText FFieldNotifyTrack::GetDisplayNameInternal() const
+{
+	return FText::FromName(FieldNotify.GetFieldName());
+}
+
+uint64 FFieldNotifyTrack::GetObjectIdInternal() const
+{
+	return ObjectId;
+}
+
+bool FFieldNotifyTrack::HandleDoubleClickInternal()
+{
+	const IRewindDebugger* RewindDebugger = IRewindDebugger::Instance();
+	if (const TraceServices::IAnalysisSession* AnalysisSession = RewindDebugger->GetAnalysisSession())
+	{
+		TraceServices::FAnalysisSessionReadScope SessionReadScope(*AnalysisSession);
+
+		const IGameplayProvider* GameplayProvider = AnalysisSession->ReadProvider<IGameplayProvider>("GameplayProvider");
+		//const FObjectInfo& AssetInfo = GameplayProvider->GetObjectInfo(GetId());
+
+		// attach the viewmodel to the UMG Preview Window for debugging
+		//if (GetPreviewWindow)
+		//{
+		//	const FObjectInfo& OwnerObjectInfo  = GameplayProvider->GetObjectInfo(ObjectId);
+		//	GetPreviewWindow->SetDebugTarget(OwnerObjectInfo.Name);
+		//	GetPreviewWindow->bEnableDebugTesting = true;
+		//}
+
+		//GEditor->GetEditorSubsystem<UAssetEditorSubsystem>()->OpenEditorForAsset(AssetInfo.PathName);
+
+		return true;
+	}
+	return false;
+}
+
+TSharedPtr<SWidget> FFieldNotifyTrack::GetTimelineViewInternal()
+{
+	return SNew(SEventTimelineView)
+		.ViewRange_Lambda([]() { return IRewindDebugger::Instance()->GetCurrentViewRange(); })
+		.EventData_Raw(this, &FFieldNotifyTrack::GetEventData);
+}
+
+		
+FObjectTrack::FObjectTrack(uint64 InObjectId)
+	: ObjectId(InObjectId)
+{
+	Icon = FSlateIcon("EditorStyle", "Sequencer.Tracks.Event", "Sequencer.Tracks.Event");
+}
+
+TSharedPtr<SWidget> FObjectTrack::GetDetailsViewInternal()
+{
+	return nullptr;
+}
+
+bool FObjectTrack::UpdateInternal()
+{
+	TRACE_CPUPROFILER_EVENT_SCOPE(FTracks::UpdateInternal);
+	IRewindDebugger* RewindDebugger = IRewindDebugger::Instance();
+	
+	TRange<double> RecordingTimeRange = RewindDebugger->GetCurrentViewRange();
+	double StartTime = RecordingTimeRange.GetLowerBoundValue();
+	double EndTime = RecordingTimeRange.GetUpperBoundValue();
+	
+	const TraceServices::IAnalysisSession* AnalysisSession = RewindDebugger->GetAnalysisSession();
+	const FTraceProvider* Provider = AnalysisSession->ReadProvider<FTraceProvider>(FTraceProvider::ProviderName);
+	
+	bool bChanged = false;
+	
+	if(Provider)
+	{
+		TArray<uint32> UniqueTrackIds;
+		
+		TraceServices::FAnalysisSessionReadScope SessionReadScope(*AnalysisSession);
+		
+		Provider->EnumerateRecordingFieldNotifies(ObjectId, StartTime, EndTime, [&UniqueTrackIds](double InStartTime, double InEndTime, uint32 InDepth, const FTraceProvider::FFieldNotifyEvent& FieldNotify)
+			{
+				UniqueTrackIds.AddUnique(FieldNotify.FieldNotifyId);
+		 	});
+		
+		UniqueTrackIds.StableSort();
+		const int32 TrackCount = UniqueTrackIds.Num();
+		
+		if (Children.Num() != TrackCount)
+		{
+			bChanged = true;
+		}
+		
+		Children.SetNum(UniqueTrackIds.Num());
+		for(int32 Index = 0; Index < TrackCount; ++Index)
+		{
+			if (!Children[Index].IsValid() || !(Children[Index].Get()->GetFieldNotifyId() == UniqueTrackIds[Index]))
+			{
+				Children[Index] = MakeShared<FFieldNotifyTrack>(ObjectId, UniqueTrackIds[Index], Provider->GetFieldNotificationId(UniqueTrackIds[Index]));
+				bChanged = true;
+			}
+		
+			if (Children[Index]->Update())
+			{
+				bChanged = true;
+			}
+		}
+	}
+	
+	return bChanged;
+}
+
+FSlateIcon FObjectTrack::GetIconInternal()
+{
+	return Icon;
+}
+
+FName FObjectTrack::GetNameInternal() const
+{
+	return "FieldNotifications";
+}
+
+FText FObjectTrack::GetDisplayNameInternal() const
+{
+	return LOCTEXT("ObjectTrackName", "Field Notify");
+}
+
+uint64 FObjectTrack::GetObjectIdInternal() const
+{
+	return ObjectId;
+}
+
+void FObjectTrack::IterateSubTracksInternal(TFunction<void(TSharedPtr<FRewindDebuggerTrack> SubTrack)> IteratorFunction)
+{
+	for(TSharedPtr<FFieldNotifyTrack>& Track : Children)
+	{
+		IteratorFunction(Track);
+	}
+}
+
+FName FTracksCreator::GetTargetTypeNameInternal() const
+{
+	static const FName ObjectName("Object");
+	return ObjectName;
+}
+
+FName FTracksCreator::GetNameInternal() const
+{
+	static const FName FieldNotificationName("FieldNotification");
+	return FieldNotificationName;
+}
+
+void FTracksCreator::GetTrackTypesInternal(TArray<RewindDebugger::FRewindDebuggerTrackType>& Types) const
+{
+	static const FName FieldNotificationName("FieldNotification");
+	Types.Add({ FieldNotificationName, LOCTEXT("FieldNotification", "Field Notification") });
+}
+
+TSharedPtr<RewindDebugger::FRewindDebuggerTrack> FTracksCreator::CreateTrackInternal(uint64 ObjectId) const
+{
+	return MakeShared<FObjectTrack>(ObjectId);
+}
+
+bool FTracksCreator::HasDebugInfoInternal(uint64 ObjectId) const
+{
+	TRACE_CPUPROFILER_EVENT_SCOPE(FTracks::HasDebugInfoInternal);
+	const TraceServices::IAnalysisSession* AnalysisSession = IRewindDebugger::Instance()->GetAnalysisSession();
+
+	TraceServices::FAnalysisSessionReadScope SessionReadScope(*AnalysisSession);
+	bool bHasData = false;
+	if (const FTraceProvider* Provider = AnalysisSession->ReadProvider<FTraceProvider>(FTraceProvider::ProviderName))
+	{
+		bHasData = Provider->HasData(ObjectId);
+	}
+	return bHasData;
+}
+	
+} // namespace
+
+#undef LOCTEXT_NAMESPACE
diff --git a/Engine/Plugins/Developer/FieldNotificationTrace/Source/FieldNotificationTraceEditor/Private/FieldNotificationTrack.h b/Engine/Plugins/Developer/FieldNotificationTrace/Source/FieldNotificationTraceEditor/Private/FieldNotificationTrack.h
new file mode 100644
--- /dev/null
+++ b/Engine/Plugins/Developer/FieldNotificationTrace/Source/FieldNotificationTraceEditor/Private/FieldNotificationTrack.h
@@ -0,0 +1,77 @@
+// Copyright Epic Games, Inc. All Rights Reserved.
+#pragma once
+
+#include "FieldNotificationId.h"
+#include "IRewindDebuggerTrackCreator.h"
+#include "RewindDebuggerTrack.h"
+#include "Textures/SlateIcon.h"
+#include "SEventTimelineView.h"
+	
+namespace UE::FieldNotification
+{
+
+class FFieldNotifyTrack : public RewindDebugger::FRewindDebuggerTrack
+{
+public:
+	FFieldNotifyTrack(uint64 InObjectId, uint32 InFieldNotifyId, FFieldNotificationId InFieldNotify);
+	
+	uint32 GetFieldNotifyId() const
+	{
+		return FieldNotifyId;
+	}
+
+private:
+	virtual bool UpdateInternal() override;
+	TSharedPtr<SWidget> GetDetailsViewInternal() override;
+	virtual TSharedPtr<SWidget> GetTimelineViewInternal() override;
+	virtual FSlateIcon GetIconInternal() override;
+	virtual FName GetNameInternal() const override;
+	virtual FText GetDisplayNameInternal() const override;
+	virtual uint64 GetObjectIdInternal() const override;
+	virtual bool HandleDoubleClickInternal() override;
+
+	TSharedPtr<SEventTimelineView::FTimelineEventData> GetEventData() const;
+	
+	mutable TSharedPtr<SEventTimelineView::FTimelineEventData> EventData;
+	mutable int EventUpdateRequested = 0;
+
+	uint64 ObjectId;
+	uint32 FieldNotifyId;
+	FFieldNotificationId FieldNotify;
+	FSlateIcon Icon;
+	FText TrackName;
+};
+
+
+class FObjectTrack : public RewindDebugger::FRewindDebuggerTrack
+{
+public:
+	FObjectTrack(uint64 InObjectId);
+private:
+	virtual bool UpdateInternal() override;
+	virtual TSharedPtr<SWidget> GetDetailsViewInternal() override;
+
+	virtual FSlateIcon GetIconInternal() override;
+	virtual FName GetNameInternal() const override;
+	virtual FText GetDisplayNameInternal() const override;
+	virtual uint64 GetObjectIdInternal() const override;
+	virtual void IterateSubTracksInternal(TFunction<void(TSharedPtr<FRewindDebuggerTrack> SubTrack)> IteratorFunction) override;
+
+	FSlateIcon Icon;
+	uint64 ObjectId;
+
+	TArray<TSharedPtr<FFieldNotifyTrack>> Children;
+};
+
+
+class FTracksCreator : public RewindDebugger::IRewindDebuggerTrackCreator
+{
+private:
+	virtual FName GetTargetTypeNameInternal() const;
+	virtual FName GetNameInternal() const override;
+	virtual void GetTrackTypesInternal(TArray<RewindDebugger::FRewindDebuggerTrackType>& Types) const override;
+	virtual TSharedPtr<RewindDebugger::FRewindDebuggerTrack> CreateTrackInternal(uint64 ObjectId) const override;
+	virtual bool HasDebugInfoInternal(uint64 ObjectId) const override;
+};
+
+} // namespace
